First pass of the Localization Service interface. Allows plugins for specific localization service providers (OneSky, etc).  Can execute commands which run asyncronously on worker thread.

diff --git a/Engine/Source/Developer/LocalizationService/LocalizationService.Build.cs b/Engine/Source/Developer/LocalizationService/LocalizationService.Build.cs
new file mode 100644
--- /dev/null
+++ b/Engine/Source/Developer/LocalizationService/LocalizationService.Build.cs
@@ -0,0 +1,50 @@
+// Copyright 1998-2015 Epic Games, Inc. All Rights Reserved.
+
+using UnrealBuildTool;
+
+public class LocalizationService : ModuleRules
+{
+	public LocalizationService(TargetInfo Target)
+	{
+        PrivateIncludePaths.Add("Developer/LocalizationService/Private");
+
+		PublicDependencyModuleNames.AddRange(
+			new string[] {
+				"Core",
+				"CoreUObject",
+				"InputCore",
+			}
+		);
+
+        if (Target.Platform != UnrealTargetPlatform.Linux)
+        {
+            PrivateDependencyModuleNames.AddRange(
+                new string[] {
+                    "Slate",
+                    "SlateCore",
+                    "EditorStyle"
+                }
+            );
+        }
+
+        if (UEBuildConfiguration.bBuildEditor)
+        {
+			PrivateDependencyModuleNames.AddRange(
+                new string[] {
+					"Engine",
+					"UnrealEd",
+                    "PropertyEditor",
+				}
+			);
+        }
+
+		if (UEBuildConfiguration.bBuildDeveloperTools)
+		{
+			PrivateDependencyModuleNames.AddRange(
+				new string[] {
+					"MessageLog",
+				}
+			);
+		}
+	}
+}
diff --git a/Engine/Source/Developer/LocalizationService/Private/DefaultLocalizationServiceProvider.cpp b/Engine/Source/Developer/LocalizationService/Private/DefaultLocalizationServiceProvider.cpp
new file mode 100644
--- /dev/null
+++ b/Engine/Source/Developer/LocalizationService/Private/DefaultLocalizationServiceProvider.cpp
@@ -0,0 +1,80 @@
+// Copyright 1998-2015 Epic Games, Inc. All Rights Reserved.
+
+#include "LocalizationServicePrivatePCH.h"
+#include "DefaultLocalizationServiceProvider.h"
+#include "ILocalizationServiceModule.h"
+#include "MessageLog.h"
+
+#define LOCTEXT_NAMESPACE "DefaultLocalizationServiceProvider"
+
+void FDefaultLocalizationServiceProvider::Init(bool bForceConnection)
+{
+	FMessageLog("LocalizationService").Info(LOCTEXT("LocalizationServiceDisabled", "Localization service is disabled"));
+}
+
+void FDefaultLocalizationServiceProvider::Close()
+{
+
+}
+
+FText FDefaultLocalizationServiceProvider::GetStatusText() const
+{
+	return LOCTEXT("LocalizationServiceDisabled", "Localization service is disabled");
+}
+
+bool FDefaultLocalizationServiceProvider::IsAvailable() const
+{
+	return false;
+}
+
+bool FDefaultLocalizationServiceProvider::IsEnabled() const
+{
+	return false;
+}
+
+const FName& FDefaultLocalizationServiceProvider::GetName(void) const
+{
+	static FName ProviderName("None"); 
+	return ProviderName; 
+}
+
+const FText FDefaultLocalizationServiceProvider::GetDisplayName() const
+{
+	return LOCTEXT("DefaultLocalizationServiceProviderDisplayName", "Localization Service Disabled");
+}
+
+ELocalizationServiceOperationCommandResult::Type FDefaultLocalizationServiceProvider::GetState(const TArray<FLocalizationServiceTranslationIdentifier>& InTranslationIds, TArray< TSharedRef<ILocalizationServiceState, ESPMode::ThreadSafe> >& OutState, ELocalizationServiceCacheUsage::Type InStateCacheUsage)
+{
+	return ELocalizationServiceOperationCommandResult::Failed;
+}
+
+ELocalizationServiceOperationCommandResult::Type FDefaultLocalizationServiceProvider::Execute(const TSharedRef<ILocalizationServiceOperation, ESPMode::ThreadSafe>& InOperation, const TArray<FLocalizationServiceTranslationIdentifier>& InTranslationIds, ELocalizationServiceOperationConcurrency::Type InConcurrency, const FLocalizationServiceOperationComplete& InOperationCompleteDelegate)
+{
+	return ELocalizationServiceOperationCommandResult::Failed;
+}
+
+bool FDefaultLocalizationServiceProvider::CanCancelOperation( const TSharedRef<ILocalizationServiceOperation, ESPMode::ThreadSafe>& InOperation ) const
+{
+	return false;
+}
+
+void FDefaultLocalizationServiceProvider::CancelOperation( const TSharedRef<ILocalizationServiceOperation, ESPMode::ThreadSafe>& InOperation )
+{
+}
+
+void FDefaultLocalizationServiceProvider::Tick()
+{
+
+}
+
+void FDefaultLocalizationServiceProvider::CustomizeSettingsDetails(IDetailCategoryBuilder& DetailCategoryBuilder) const
+{
+
+}
+
+void FDefaultLocalizationServiceProvider::CustomizeTargetDetails(IDetailCategoryBuilder& DetailCategoryBuilder, ULocalizationTarget& Target) const
+{
+
+}
+
+#undef LOCTEXT_NAMESPACE
\ No newline at end of file
diff --git a/Engine/Source/Developer/LocalizationService/Private/DefaultLocalizationServiceProvider.h b/Engine/Source/Developer/LocalizationService/Private/DefaultLocalizationServiceProvider.h
new file mode 100644
--- /dev/null
+++ b/Engine/Source/Developer/LocalizationService/Private/DefaultLocalizationServiceProvider.h
@@ -0,0 +1,31 @@
+// Copyright 1998-2015 Epic Games, Inc. All Rights Reserved.
+
+#pragma once
+
+#include "ILocalizationServiceProvider.h"
+
+/**
+ * Default localization service provider implementation - "None"
+ */
+class FDefaultLocalizationServiceProvider : public ILocalizationServiceProvider
+{
+public:
+	// ILocalizationServiceProvider implementation
+	virtual void Init(bool bForceConnection = true) override;
+	virtual void Close() override;
+	virtual FText GetStatusText() const override;
+	virtual bool IsAvailable() const override;
+	virtual bool IsEnabled() const override;
+	virtual const FName& GetName(void) const override;
+	virtual const FText GetDisplayName() const override;
+	virtual ELocalizationServiceOperationCommandResult::Type GetState(const TArray<FLocalizationServiceTranslationIdentifier>& InTranslationIds, TArray< TSharedRef<ILocalizationServiceState, ESPMode::ThreadSafe> >& OutState, ELocalizationServiceCacheUsage::Type InStateCacheUsage) override;
+	virtual ELocalizationServiceOperationCommandResult::Type Execute(const TSharedRef<ILocalizationServiceOperation, ESPMode::ThreadSafe>& InOperation, const TArray<FLocalizationServiceTranslationIdentifier>& InTranslationIds, ELocalizationServiceOperationConcurrency::Type InConcurrency = ELocalizationServiceOperationConcurrency::Synchronous, const FLocalizationServiceOperationComplete& InOperationCompleteDelegate = FLocalizationServiceOperationComplete()) override;
+	virtual bool CanCancelOperation( const TSharedRef<ILocalizationServiceOperation, ESPMode::ThreadSafe>& InOperation ) const override;
+	virtual void CancelOperation( const TSharedRef<ILocalizationServiceOperation, ESPMode::ThreadSafe>& InOperation ) override;
+	virtual void Tick() override;
+#if LOCALIZATION_SERVICES_WITH_SLATE
+	virtual void CustomizeSettingsDetails(IDetailCategoryBuilder& DetailCategoryBuilder) const override;
+	virtual void CustomizeTargetDetails(IDetailCategoryBuilder& DetailCategoryBuilder, ULocalizationTarget& Target) const override;
+#endif // LOCALIZATION_SERVICES_WITH_SLATE
+};
+
diff --git a/Engine/Source/Developer/LocalizationService/Private/LocalizationServiceHelpers.cpp b/Engine/Source/Developer/LocalizationService/Private/LocalizationServiceHelpers.cpp
new file mode 100644
--- /dev/null
+++ b/Engine/Source/Developer/LocalizationService/Private/LocalizationServiceHelpers.cpp
@@ -0,0 +1,92 @@
+// Copyright 1998-2015 Epic Games, Inc. All Rights Reserved.
+
+#include "LocalizationServicePrivatePCH.h"
+#include "LocalizationServiceHelpers.h"
+#include "ILocalizationServiceProvider.h"
+#include "ILocalizationServiceModule.h"
+#include "MessageLog.h"
+
+#define LOCTEXT_NAMESPACE "LocalizationServiceHelpers"
+
+namespace LocalizationServiceHelpers
+{
+
+	const FString& GetSettingsIni()
+	{
+		if (ILocalizationServiceModule::Get().GetUseGlobalSettings())
+		{
+			return GetGlobalSettingsIni();
+		}
+		else
+		{
+			static FString LocalizationServiceSettingsIni;
+			if (LocalizationServiceSettingsIni.Len() == 0)
+			{
+				const FString LocalizationServiceSettingsDir = FPaths::GeneratedConfigDir();
+				FConfigCacheIni::LoadGlobalIniFile(LocalizationServiceSettingsIni, TEXT("LocalizationServiceSettings"), NULL, false, false, true, *LocalizationServiceSettingsDir);
+			}
+			return LocalizationServiceSettingsIni;
+		}
+	}
+
+	const FString& GetGlobalSettingsIni()
+	{
+		static FString LocalizationServiceGlobalSettingsIni;
+		if (LocalizationServiceGlobalSettingsIni.Len() == 0)
+		{
+			const FString LocalizationServiceSettingsDir = FPaths::EngineSavedDir() + TEXT("Config/");
+			FConfigCacheIni::LoadGlobalIniFile(LocalizationServiceGlobalSettingsIni, TEXT("LocalizationServiceSettings"), NULL, false, false, true, *LocalizationServiceSettingsDir);
+		}
+		return LocalizationServiceGlobalSettingsIni;
+	}
+
+
+	//LOCALIZATIONSERVICE_API extern bool CommitTranslation(const FString& Culture, const FString& Namespace, const FString& Source, const FString& Translation)
+	//{
+	//	if (!Culture.IsValid() || Culture->GetName().IsEmpty())
+	//	{
+	//		FMessageLog("LocalizationService").Error(LOCTEXT("CultureNotSpecified", "Culture not specified"));
+	//		return false;
+	//	}
+
+	//	if (!ILocalizationServiceModule::Get().IsEnabled())
+	//	{
+	//		FMessageLog("LocalizationService").Error(LOCTEXT("LocalizationServiceDisabled", "Localization Service is not enabled."));
+	//		return false;
+	//	}
+
+	//	if (!ILocalizationServiceModule::Get().GetProvider().IsAvailable())
+	//	{
+	//		FMessageLog("SourceControl").Error(LOCTEXT("ILocalizationServiceUnavailable", "Localization Ser is currently not available."));
+	//		return false;
+	//	}
+
+	//	bool bSuccessfullyCheckedOut = false;
+
+	//	ILocalizationServiceProvider& Provider = ILocalizationServiceModule::Get().GetProvider();
+	//	//TArray<TSharedRef<ILocalizationServiceState, ESPMode::ThreadSafe>> TranslationStates;
+	//	TSharedPtr<ILocalizationServiceState, ESPMode::ThreadSafe> LocalizationServiceState = Provider.GetState(FLocalizationServiceTranslationIdentifier(Culture, Namespace, Source), ELocalizationServiceCacheUsage::ForceUpdate);
+
+	//	if (LocalizationServiceState.IsValid())
+	//	{
+	//		// TODO: Finish
+	//	}
+	//}
+}
+
+FScopedLocalizationService::FScopedLocalizationService()
+{
+	ILocalizationServiceModule::Get().GetProvider().Init();
+}
+
+FScopedLocalizationService::~FScopedLocalizationService()
+{
+	ILocalizationServiceModule::Get().GetProvider().Close();
+}
+
+ILocalizationServiceProvider& FScopedLocalizationService::GetProvider()
+{
+	return ILocalizationServiceModule::Get().GetProvider();
+}
+
+#undef LOCTEXT_NAMESPACE
\ No newline at end of file
diff --git a/Engine/Source/Developer/LocalizationService/Private/LocalizationServiceModule.cpp b/Engine/Source/Developer/LocalizationService/Private/LocalizationServiceModule.cpp
new file mode 100644
--- /dev/null
+++ b/Engine/Source/Developer/LocalizationService/Private/LocalizationServiceModule.cpp
@@ -0,0 +1,255 @@
+// Copyright 1998-2015 Epic Games, Inc. All Rights Reserved.
+
+#include "LocalizationServicePrivatePCH.h"
+#include "LocalizationServiceModule.h"
+#include "LocalizationServiceSettings.h"
+#include "Runtime/Core/Public/Features/IModularFeatures.h"
+#include "Developer/MessageLog/Public/MessageLogModule.h"
+
+#if WITH_EDITOR
+	//#include "Settings/EditorLoadingSavingSettings.h"
+	//#include "Runtime/Engine/Public/EngineAnalytics.h"
+	//#include "Runtime/Analytics/Analytics/Public/Interfaces/IAnalyticsProvider.h"
+#endif
+
+DEFINE_LOG_CATEGORY(LogLocalizationService);
+
+#define LOCTEXT_NAMESPACE "LocalizationService"
+
+static const FName LocalizationServiceFeatureName("LocalizationService");
+
+namespace LocalizationServiceConstants
+{
+	/** The maximum number of translation status requests we should dispatch in a tick */
+	const int32 MaxStatusDispatchesPerTick = 64;
+}
+
+FLocalizationServiceModule::FLocalizationServiceModule()
+	: CurrentLocalizationServiceProvider(NULL)
+{
+}
+
+void FLocalizationServiceModule::StartupModule()
+{
+	// load our settings
+	LocalizationServiceSettings.LoadSettings();
+
+	// Register to check for translation service features
+	IModularFeatures::Get().OnModularFeatureRegistered().AddRaw(this, &FLocalizationServiceModule::HandleModularFeatureRegistered);
+	IModularFeatures::Get().OnModularFeatureUnregistered().AddRaw(this, &FLocalizationServiceModule::HandleModularFeatureUnregistered);
+
+	// bind default provider to editor
+	IModularFeatures::Get().RegisterModularFeature( LocalizationServiceFeatureName, &DefaultLocalizationServiceProvider );
+
+#if WITH_UNREAL_DEVELOPER_TOOLS
+	// create a message log for translation service to use
+	FMessageLogModule& MessageLogModule = FModuleManager::LoadModuleChecked<FMessageLogModule>("MessageLog");
+	MessageLogModule.RegisterLogListing("LocalizationService", LOCTEXT("LocalizationServiceLogLabel", "Localization Service"));
+#endif
+}
+
+void FLocalizationServiceModule::ShutdownModule()
+{
+	// close the current provider
+	GetProvider().Close();
+
+#if WITH_UNREAL_DEVELOPER_TOOLS
+	// unregister message log
+	if(FModuleManager::Get().IsModuleLoaded("MessageLog"))
+	{
+		FMessageLogModule& MessageLogModule = FModuleManager::LoadModuleChecked<FMessageLogModule>("MessageLog");
+		MessageLogModule.UnregisterLogListing("LocalizationService");
+	}
+#endif
+
+	// unbind default provider from editor
+	IModularFeatures::Get().UnregisterModularFeature( LocalizationServiceFeatureName, &DefaultLocalizationServiceProvider );
+
+	// we don't care about modular features any more
+	IModularFeatures::Get().OnModularFeatureRegistered().RemoveAll(this);
+	IModularFeatures::Get().OnModularFeatureUnregistered().RemoveAll(this);
+}
+
+void FLocalizationServiceModule::SaveSettings()
+{
+	LocalizationServiceSettings.SaveSettings();
+}
+
+void FLocalizationServiceModule::InitializeLocalizationServiceProviders()
+{
+	int32 LocalizationServiceCount = IModularFeatures::Get().GetModularFeatureImplementationCount(LocalizationServiceFeatureName);
+	if( LocalizationServiceCount > 0 )
+	{
+		FString PreferredLocalizationServiceProvider = LocalizationServiceSettings.GetProvider();
+		TArray<ILocalizationServiceProvider*> Providers = IModularFeatures::Get().GetModularFeatureImplementations<ILocalizationServiceProvider>(LocalizationServiceFeatureName);
+		for(auto It(Providers.CreateIterator()); It; It++)
+		{ 
+			ILocalizationServiceProvider* Provider = *It;
+			if(PreferredLocalizationServiceProvider == Provider->GetName().ToString())
+			{
+				CurrentLocalizationServiceProvider = Provider;
+				break;
+			}
+		}
+
+		// no provider found of this name, default to the first one
+		if( CurrentLocalizationServiceProvider == NULL )
+		{
+			CurrentLocalizationServiceProvider = &DefaultLocalizationServiceProvider;
+		}
+	}
+
+	check(CurrentLocalizationServiceProvider);
+
+	CurrentLocalizationServiceProvider->Init(false);	// Don't force a connection here, as its synchronous. Let the user establish a connection.
+}
+
+void FLocalizationServiceModule::Tick()
+{	
+	ILocalizationServiceProvider& Provider = GetProvider();
+
+	// tick the provider, so any operation results can be read back
+	Provider.Tick();
+	
+	// don't allow background status updates when disabled 
+	if(Provider.IsEnabled())
+	{
+		//// check for any pending dispatches
+		//if(PendingStatusUpdateTranslations.Num() > 0)
+		//{
+		//	// grab a batch of translations
+		//	TArray<FLocalizationServiceTranslationIdentifier> TranslationsToDispatch;
+		//	for(auto Iter(PendingStatusUpdateTranslations.CreateConstIterator()); Iter; Iter++)
+		//	{
+		//		if(FilesToDispatch.Num() >= LocalizationServiceConstants::MaxStatusDispatchesPerTick)
+		//		{
+		//			break;
+		//		}
+		//		FilesToDispatch.Add(*Iter);
+		//	}
+
+		//	if(FilesToDispatch.Num() > 0)
+		//	{
+		//		// remove the files we are dispatching so we don't try again
+		//		PendingStatusUpdateTranslations.RemoveAt(0, TranslationsToDispatch.Num());
+
+		//		// dispatch update
+		//		Provider.Execute(ILocalizationServiceOperation::Create<FUpdateStatus>(), TranslationsToDispatch, EConcurrency::Asynchronous);
+		//	}
+		//}
+	}
+}
+
+
+//void FLocalizationServiceModule::QueueStatusUpdate(const TArray<FLocalizationServiceTranslationIdentifier>& InTranslationIds)
+//{
+//	if(IsEnabled())
+//	{
+//		for(auto It(InTranslationIds.CreateConstIterator()); It; It++)
+//		{
+//			QueueStatusUpdate(*It);
+//		}
+//	}
+//}
+
+bool FLocalizationServiceModule::IsEnabled() const
+{
+	return GetProvider().IsEnabled();
+}
+
+ILocalizationServiceProvider& FLocalizationServiceModule::GetProvider() const
+{
+	return *CurrentLocalizationServiceProvider;
+}
+
+void FLocalizationServiceModule::SetProvider( const FName& InName )
+{
+	TArray<ILocalizationServiceProvider*> Providers = IModularFeatures::Get().GetModularFeatureImplementations<ILocalizationServiceProvider>(LocalizationServiceFeatureName);
+	for(auto It(Providers.CreateIterator()); It; It++)
+	{
+		ILocalizationServiceProvider* Provider = *It;
+		if(InName == Provider->GetName())
+		{
+			SetCurrentLocalizationServiceProvider(*Provider);
+			return;
+		}
+	}
+
+	UE_LOG(LogLocalizationService, Fatal, TEXT("Tried to set unknown translation service provider: %s"), *InName.ToString());
+}
+
+void FLocalizationServiceModule::ClearCurrentLocalizationServiceProvider()
+{
+	if( CurrentLocalizationServiceProvider != NULL )
+	{
+		CurrentLocalizationServiceProvider->Close();
+		CurrentLocalizationServiceProvider = &DefaultLocalizationServiceProvider;
+	}
+}
+
+int32 FLocalizationServiceModule::GetNumLocalizationServiceProviders()
+{
+	return IModularFeatures::Get().GetModularFeatureImplementationCount(LocalizationServiceFeatureName);
+}
+
+void FLocalizationServiceModule::SetCurrentLocalizationServiceProvider(int32 ProviderIndex)
+{
+	TArray<ILocalizationServiceProvider*> Providers = IModularFeatures::Get().GetModularFeatureImplementations<ILocalizationServiceProvider>(LocalizationServiceFeatureName);
+	check(Providers.IsValidIndex(ProviderIndex));
+	SetCurrentLocalizationServiceProvider(*Providers[ProviderIndex]);
+}
+
+void FLocalizationServiceModule::SetCurrentLocalizationServiceProvider(ILocalizationServiceProvider& InProvider)
+{
+	// see if we are switching or not
+	if(&InProvider == CurrentLocalizationServiceProvider)
+	{
+		return;
+	}
+
+	ClearCurrentLocalizationServiceProvider();
+
+	CurrentLocalizationServiceProvider = &InProvider;
+	CurrentLocalizationServiceProvider->Init(false);	// Don't force a connection here, as its synchronous. Let the user establish a connection.
+
+	LocalizationServiceSettings.SetProvider(CurrentLocalizationServiceProvider->GetName().ToString());
+
+	SaveSettings();
+}
+
+FName FLocalizationServiceModule::GetLocalizationServiceProviderName(int32 ProviderIndex)
+{
+	TArray<ILocalizationServiceProvider*> Providers = IModularFeatures::Get().GetModularFeatureImplementations<ILocalizationServiceProvider>(LocalizationServiceFeatureName);
+	check(Providers.IsValidIndex(ProviderIndex));
+	return Providers[ProviderIndex]->GetName();
+}
+
+void FLocalizationServiceModule::HandleModularFeatureRegistered(const FName& Type, IModularFeature* ModularFeature)
+{
+	if(Type == LocalizationServiceFeatureName)
+	{
+		InitializeLocalizationServiceProviders();
+	}
+}
+
+void FLocalizationServiceModule::HandleModularFeatureUnregistered(const FName& Type, IModularFeature* ModularFeature)
+{
+	if(Type == LocalizationServiceFeatureName && CurrentLocalizationServiceProvider == static_cast<ILocalizationServiceProvider*>(ModularFeature))
+	{
+		ClearCurrentLocalizationServiceProvider();
+	}
+}
+
+bool FLocalizationServiceModule::GetUseGlobalSettings() const
+{
+	return LocalizationServiceSettings.GetUseGlobalSettings();
+}
+
+void FLocalizationServiceModule::SetUseGlobalSettings(bool bIsUseGlobalSettings)
+{
+	LocalizationServiceSettings.SetUseGlobalSettings(bIsUseGlobalSettings);
+}	
+
+IMPLEMENT_MODULE( FLocalizationServiceModule, LocalizationService );
+
+#undef LOCTEXT_NAMESPACE
diff --git a/Engine/Source/Developer/LocalizationService/Private/LocalizationServiceModule.h b/Engine/Source/Developer/LocalizationService/Private/LocalizationServiceModule.h
new file mode 100644
--- /dev/null
+++ b/Engine/Source/Developer/LocalizationService/Private/LocalizationServiceModule.h
@@ -0,0 +1,88 @@
+// Copyright 1998-2015 Epic Games, Inc. All Rights Reserved.
+
+#pragma once
+
+/**
+ * Interface for talking to source control clients
+ */
+
+#include "ModuleManager.h"
+#include "ILocalizationServiceModule.h"
+#include "ILocalizationServiceProvider.h"
+#include "LocalizationServiceSettings.h"
+#include "DefaultLocalizationServiceProvider.h"
+
+class FLocalizationServiceModule : public ILocalizationServiceModule
+{
+public:
+	FLocalizationServiceModule();
+
+	/** IModuleInterface implementation */
+	virtual void StartupModule() override;
+	virtual void ShutdownModule() override;
+
+	/** ILocalizationServiceModule implementation */
+	virtual void Tick() override;
+	virtual bool IsEnabled() const override;
+	virtual ILocalizationServiceProvider& GetProvider() const override;
+	virtual void SetProvider( const FName& InName ) override;
+	virtual bool GetUseGlobalSettings() const override;
+	virtual void SetUseGlobalSettings(bool bIsUseGlobalSettings) override;
+
+	/** Save the settings to the ini file */
+	void SaveSettings();
+
+	/**
+	 * Get the number of currently registered translation service providers.
+	 */
+	int32 GetNumLocalizationServiceProviders();
+
+	/**
+	 * Set the current translation service provider by index.
+	 */
+	void SetCurrentLocalizationServiceProvider(int32 ProviderIndex);
+
+	/**
+	 * Get the name of the translation service provider at the specified index.
+	 */
+	FName GetLocalizationServiceProviderName(int32 ProviderIndex);
+
+	/**
+	 * Gets a reference to the translation service module instance.
+	 *
+	 * @return A reference to the translation service module.
+	 */
+	static FLocalizationServiceModule& Get()
+	{
+		return FModuleManager::LoadModuleChecked<FLocalizationServiceModule>("LocalizationService");
+	}
+
+private:
+	/** Refresh & initialize the current translation service provider */
+	void InitializeLocalizationServiceProviders();
+
+	/** Close the current translation service provider & set the current to default - 'None' */
+	void ClearCurrentLocalizationServiceProvider();
+
+	/** Set the current translation service provider to the passed-in value */
+	void SetCurrentLocalizationServiceProvider(ILocalizationServiceProvider& InProvider);
+
+	/** Delegate handling when translation service features are registered */
+	void HandleModularFeatureRegistered(const FName& Type, IModularFeature* ModularFeature);
+
+	/** Delegate handling when translation service features are unregistered */
+	void HandleModularFeatureUnregistered(const FName& Type, IModularFeature* ModularFeature);
+
+private:
+	/** The settings object */
+	FLocalizationServiceSettings LocalizationServiceSettings;
+
+	/** Current translation service provider */
+	ILocalizationServiceProvider* CurrentLocalizationServiceProvider;
+
+	/** Source control provider we use if there are none registered */
+	FDefaultLocalizationServiceProvider DefaultLocalizationServiceProvider;
+
+	///** Translations pending a status update */
+	//TArray < FLocalizationServiceTranslationIdentifier > PendingStatusUpdateTranslations;
+};
diff --git a/Engine/Source/Developer/LocalizationService/Private/LocalizationServicePrivatePCH.h b/Engine/Source/Developer/LocalizationService/Private/LocalizationServicePrivatePCH.h
new file mode 100644
--- /dev/null
+++ b/Engine/Source/Developer/LocalizationService/Private/LocalizationServicePrivatePCH.h
@@ -0,0 +1,26 @@
+// Copyright 1998-2015 Epic Games, Inc. All Rights Reserved.
+
+#pragma once
+
+#include "Core.h"
+#include "ILocalizationServiceProvider.h"	// for TRANSLATION_SERVICES_WITH_SLATE
+
+#ifndef LOCALIZATION_SERVICES_WITH_SLATE
+	#error LOCALIZATION_SERVICES_WITH_SLATE not defined
+#endif // LOCALIZATION_SERVICES_WITH_SLATE
+
+#if LOCALIZATION_SERVICES_WITH_SLATE
+	#include "SlateBasics.h"
+	#include "EditorStyle.h"
+#endif // LOCALIZATION_SERVICES_WITH_SLATE
+
+#if WITH_EDITOR
+#include "Engine.h"
+#include "PackageTools.h"
+#include "MessageLog.h"
+#endif
+
+#if WITH_UNREAL_DEVELOPER_TOOLS
+#include "MessageLogModule.h"
+#endif
+ 
\ No newline at end of file
diff --git a/Engine/Source/Developer/LocalizationService/Private/LocalizationServiceSettings.cpp b/Engine/Source/Developer/LocalizationService/Private/LocalizationServiceSettings.cpp
new file mode 100644
--- /dev/null
+++ b/Engine/Source/Developer/LocalizationService/Private/LocalizationServiceSettings.cpp
@@ -0,0 +1,75 @@
+// Copyright 1998-2015 Epic Games, Inc. All Rights Reserved.
+
+#include "LocalizationServicePrivatePCH.h"
+#include "LocalizationServiceSettings.h"
+#include "ILocalizationServiceModule.h"
+#include "LocalizationServiceHelpers.h"
+
+namespace LocalizationServiceSettingsConstants
+{
+
+/** The section of the ini file we load our settings from */
+static const FString SettingsSection = TEXT("LocalizationService.LocalizationServiceSettings");
+
+}
+
+
+const FString& FLocalizationServiceSettings::GetProvider() const
+{
+	return Provider;
+}
+
+void FLocalizationServiceSettings::SetProvider(const FString& InString)
+{
+	Provider = InString;
+}
+
+bool FLocalizationServiceSettings::GetUseGlobalSettings() const
+{
+	return bUseGlobalSettings;
+}
+
+void FLocalizationServiceSettings::SetUseGlobalSettings(bool bInUseGlobalSettings)
+{
+	bUseGlobalSettings = bInUseGlobalSettings;
+}
+
+void FLocalizationServiceSettings::LoadSettings()
+{
+	// make sure we load the global ini first
+	const FString& GlobalIniFile = LocalizationServiceHelpers::GetGlobalSettingsIni();
+	GConfig->GetBool(*LocalizationServiceSettingsConstants::SettingsSection, TEXT("UseGlobalSettings"), bUseGlobalSettings, GlobalIniFile);
+
+	TArray<FString> Tokens;
+	TArray<FString> Switches;
+	FCommandLine::Parse( FCommandLine::Get(), Tokens, Switches );
+	TMap<FString, FString> SwitchPairs;
+	for (int32 SwitchIdx = Switches.Num() - 1; SwitchIdx >= 0; --SwitchIdx)
+	{
+		FString& Switch = Switches[SwitchIdx];
+		TArray<FString> SplitSwitch;
+		if (2 == Switch.ParseIntoArray(SplitSwitch, TEXT("="), true))
+		{
+			SwitchPairs.Add(SplitSwitch[0], SplitSwitch[1].TrimQuotes());
+			Switches.RemoveAt(SwitchIdx);
+		}
+	}
+
+	if( SwitchPairs.Contains( TEXT("LocalizationServiceProvider") ) )
+	{
+		Provider = SwitchPairs[TEXT("LocalizationServiceProvider")];
+	}
+	else
+	{
+		const FString& IniFile = LocalizationServiceHelpers::GetSettingsIni();
+		GConfig->GetString(*LocalizationServiceSettingsConstants::SettingsSection, TEXT("Provider"), Provider, IniFile);
+	}
+}
+
+void FLocalizationServiceSettings::SaveSettings() const
+{
+	const FString& IniFile = LocalizationServiceHelpers::GetSettingsIni();
+	const FString& GlobalIniFile = LocalizationServiceHelpers::GetGlobalSettingsIni();
+	GConfig->SetString(*LocalizationServiceSettingsConstants::SettingsSection, TEXT("Provider"), *Provider, IniFile);
+	GConfig->SetBool(*LocalizationServiceSettingsConstants::SettingsSection, TEXT("UseGlobalSettings"), bUseGlobalSettings, GlobalIniFile);	
+}
\ No newline at end of file
diff --git a/Engine/Source/Developer/LocalizationService/Private/LocalizationServiceSettings.h b/Engine/Source/Developer/LocalizationService/Private/LocalizationServiceSettings.h
new file mode 100644
--- /dev/null
+++ b/Engine/Source/Developer/LocalizationService/Private/LocalizationServiceSettings.h
@@ -0,0 +1,37 @@
+// Copyright 1998-2015 Epic Games, Inc. All Rights Reserved.
+
+#pragma once
+
+class FLocalizationServiceSettings
+{
+public:
+	FLocalizationServiceSettings()
+		: bUseGlobalSettings(false)
+	{
+	}
+
+	/** Get the provider we want to use */
+	const FString& GetProvider() const;
+
+	/** Set the provider we want to use */
+	void SetProvider(const FString& InString);
+
+	/** Get whether we should use global or per-project settings */
+	bool GetUseGlobalSettings() const;
+
+	/** Set whether we should use global or per-project settings */
+	void SetUseGlobalSettings(bool bInUseGlobalSettings);
+
+	/** Load settings from ini file */
+	void LoadSettings();
+
+	/** Save settings to ini file */
+	void SaveSettings() const;
+
+private:
+	/** The preferred Localization service provider */
+	FString Provider;
+
+	/** Whether we should use global or per-project settings */
+	bool bUseGlobalSettings;
+};
\ No newline at end of file
diff --git a/Engine/Source/Developer/LocalizationService/Private/ScopedLocalizationServiceProgress.cpp b/Engine/Source/Developer/LocalizationService/Private/ScopedLocalizationServiceProgress.cpp
new file mode 100644
--- /dev/null
+++ b/Engine/Source/Developer/LocalizationService/Private/ScopedLocalizationServiceProgress.cpp
@@ -0,0 +1,239 @@
+// Copyright 1998-2015 Epic Games, Inc. All Rights Reserved.
+
+#include "LocalizationServicePrivatePCH.h"
+#include "ScopedLocalizationServiceProgress.h"
+#if LOCALIZATION_SERVICES_WITH_SLATE
+#include "SThrobber.h"
+
+#define LOCTEXT_NAMESPACE "LocalizationServiceProgress"
+
+namespace LocalizationServiceConstants
+{
+	/** The time (in seconds) we wait before letting the user know that an operation is taking a while */
+	float OperationTimeOut = 10.0f;
+}
+
+class SLocalizationServiceProgress : public SBorder
+{
+public:
+	SLATE_BEGIN_ARGS(SLocalizationServiceProgress) {}
+	
+	/** The text to display */
+	SLATE_ATTRIBUTE(FText, Text)
+
+	/** The delegate to call when the cancel button is clicked */
+	SLATE_ARGUMENT(FSimpleDelegate, OnCancelled)
+
+	SLATE_END_ARGS()
+
+public:
+	void Construct(const FArguments& InArgs)
+	{
+		OnCancelled = InArgs._OnCancelled;
+		bCancelClicked = false;
+		TimeStamp = FPlatformTime::Seconds();
+
+		SBorder::Construct( SBorder::FArguments()
+		.BorderImage( FEditorStyle::GetBrush("ChildWindow.Background") )
+		.Padding(16.0f)
+		.VAlign(VAlign_Center)
+		[
+			SNew(SBox)
+			.WidthOverride(500.0f)
+			[
+				SNew(SVerticalBox)
+				+SVerticalBox::Slot()
+				.AutoHeight()
+				.Padding(0.0f, 8.0f)
+				[
+					SNew(STextBlock)
+					.Text(InArgs._Text)
+				]
+				+SVerticalBox::Slot()
+				.AutoHeight()
+				.Padding(0.0f, 8.0f)
+				[
+					SNew(SHorizontalBox)
+					.Visibility(this, &SLocalizationServiceProgress::GetWarningVisibility)
+					+SHorizontalBox::Slot()
+					.AutoWidth()
+					.VAlign(VAlign_Center)
+					.HAlign(HAlign_Center)
+					[
+						SNew(SImage)
+						.Image(FEditorStyle::GetBrush("LocalizationService.ProgressWindow.Warning"))
+					]
+					+SHorizontalBox::Slot()
+					.VAlign(VAlign_Center)
+					.AutoWidth()
+					[
+						SNew(STextBlock)
+						.Text(LOCTEXT("LongTaskMessage", "Operation is taking a long time to complete. Click cancel to stop the current operation, you can try again later."))
+						.WrapTextAt(450.0f)
+					]
+				]
+				+SVerticalBox::Slot()
+				.AutoHeight()
+				.Padding(0.0f, 8.0f, 0.0f, 0.0f)
+				[
+					SNew(SHorizontalBox)
+					+SHorizontalBox::Slot()
+					.FillWidth(1.0f)
+					.VAlign(VAlign_Center)
+					.HAlign(HAlign_Left)
+					[
+						SNew(SThrobber)
+					]
+					+SHorizontalBox::Slot()
+					.HAlign(HAlign_Right)
+					.AutoWidth()
+					[
+						// buttons
+						SNew(SUniformGridPanel)
+						.SlotPadding(FEditorStyle::GetMargin("StandardDialog.SlotPadding"))
+						.MinDesiredSlotWidth(FEditorStyle::GetFloat("StandardDialog.MinDesiredSlotWidth"))
+						.MinDesiredSlotHeight(FEditorStyle::GetFloat("StandardDialog.MinDesiredSlotHeight"))
+						+SUniformGridPanel::Slot(0, 0)
+						[
+							SNew(SButton)
+							.HAlign(HAlign_Center)
+							.Text(LOCTEXT("ContinueButtonLabel", "Continue"))
+							.OnClicked(this, &SLocalizationServiceProgress::OnContinueClicked)
+							.Visibility(this, &SLocalizationServiceProgress::GetWarningVisibility)
+						]
+						+SUniformGridPanel::Slot(1, 0)
+						[
+							SNew(SButton)
+							.HAlign(HAlign_Center)
+							.Text(LOCTEXT("CancelButtonLabel", "Cancel"))
+							.OnClicked(this, &SLocalizationServiceProgress::OnCancelClicked)
+							.IsEnabled(this, &SLocalizationServiceProgress::IsCancelEnabled)
+							.Visibility(this, &SLocalizationServiceProgress::GetCancelVisibility)
+						]
+					]
+				]
+			]
+		]);
+	}
+
+private:
+
+	/** Delegate used to get the visibility of the cancel button */
+	EVisibility GetCancelVisibility() const
+	{
+		return OnCancelled.IsBound() ? EVisibility::Visible : EVisibility::Collapsed;
+	}
+
+	/** Delegate used to get the visibility of the warning text */
+	EVisibility GetWarningVisibility() const
+	{
+		if(!OnCancelled.IsBound())
+		{
+			return EVisibility::Collapsed;
+		}
+		else if(FPlatformTime::Seconds() - TimeStamp > LocalizationServiceConstants::OperationTimeOut)
+		{
+			return EVisibility::Visible;
+		}
+
+		return EVisibility::Hidden;
+	}
+
+	/** Handler for the continue button */
+	FReply OnContinueClicked()
+	{
+		TimeStamp = FPlatformTime::Seconds();
+		return FReply::Handled();
+	}
+
+	/** Handler for the cancel button */
+	FReply OnCancelClicked()
+	{
+		bCancelClicked = true;
+		OnCancelled.ExecuteIfBound();
+		return FReply::Handled();
+	}
+
+	/** Delegate used to get the enabled state of the cancel button */
+	bool IsCancelEnabled() const
+	{
+		return !bCancelClicked;
+	}
+
+private:
+
+	/** Flag used to disable the cancel button once clicked */
+	bool bCancelClicked;
+
+	/** The timer we use to determine when to display the 'long task' message */
+	float TimeStamp;
+
+	/** The delegate to call when the cancel button is clicked */
+	FSimpleDelegate OnCancelled;
+};
+
+FScopedLocalizationServiceProgress::FScopedLocalizationServiceProgress(const FText& InText, const FSimpleDelegate& InOnCancelled)
+{
+	if(!(FApp::IsUnattended() || IsRunningCommandlet()) && !InText.IsEmpty())
+	{
+		TSharedRef<SWindow> Window = SNew(SWindow)
+			.IsPopupWindow(true)
+			.SupportsMaximize(false) 
+			.SupportsMinimize(false)
+			.CreateTitleBar(false)
+			.SizingRule(ESizingRule::Autosized);
+
+		WindowPtr = Window;
+
+		Window->SetContent(
+			SNew(SLocalizationServiceProgress)
+			.Text(InText)
+			.OnCancelled(InOnCancelled) 
+		);
+			
+		TSharedPtr<SWindow> RootWindow = FGlobalTabmanager::Get()->GetRootWindow();
+		FSlateApplication::Get().AddModalWindow(Window, RootWindow, true);
+
+		Window->ShowWindow();	
+	
+		Tick();
+	}
+}
+
+FScopedLocalizationServiceProgress::~FScopedLocalizationServiceProgress()
+{
+	if(WindowPtr.IsValid())
+	{
+		WindowPtr.Pin()->RequestDestroyWindow();
+	}
+}
+
+void FScopedLocalizationServiceProgress::Tick()
+{
+	if (!(FApp::IsUnattended() || IsRunningCommandlet()) && WindowPtr.IsValid() && FSlateApplication::Get().CanDisplayWindows())
+	{
+		// Tick Slate application
+		FSlateApplication::Get().Tick();
+
+		// Sync the game thread and the render thread
+		FSlateApplication::Get().GetRenderer()->Sync();
+	}
+}
+#else
+FScopedLocalizationServiceProgress::FScopedLocalizationServiceProgress(const FText& InText, const FSimpleDelegate& InOnCancelled)
+{
+
+}
+
+FScopedLocalizationServiceProgress::~FScopedLocalizationServiceProgress()
+{
+
+}
+
+void FScopedLocalizationServiceProgress::Tick()
+{
+
+}
+#endif // LOCALIZATION_SERVICES_WITH_SLATE
+
+#undef LOCTEXT_NAMESPACE
diff --git a/Engine/Source/Developer/LocalizationService/Public/ILocalizationServiceModule.h b/Engine/Source/Developer/LocalizationService/Public/ILocalizationServiceModule.h
new file mode 100644
--- /dev/null
+++ b/Engine/Source/Developer/LocalizationService/Public/ILocalizationServiceModule.h
@@ -0,0 +1,97 @@
+// Copyright 1998-2015 Epic Games, Inc. All Rights Reserved.
+
+#pragma once
+
+#include "ModuleManager.h"
+#include "ILocalizationServiceProvider.h"
+
+LOCALIZATIONSERVICE_API DECLARE_LOG_CATEGORY_EXTERN(LogLocalizationService, Log, All);
+
+/** Delegate called when the localization service login window is closed. Parameter determines if service is enabled or not */
+DECLARE_DELEGATE_OneParam( FLocalizationServiceLoginClosed, bool );
+
+/**
+ * The modality of the login window.
+ */
+namespace ELocalizationServiceLoginWindowMode
+{
+	enum Type
+	{
+		Modal,
+		Modeless
+	};
+};
+
+
+/**
+ * Login window startup behavior
+ */
+namespace ELocalizationServiceOnLoginWindowStartup
+{
+	enum Type
+	{
+		ResetProviderToNone,
+		PreserveProvider
+	};
+};
+
+
+/**
+ * Interface for talking to localization service providers
+ */
+class ILocalizationServiceModule : public IModuleInterface
+{
+public:
+
+	/**
+	 * Tick the localization service module.
+	 * This is responsible for dispatching batched/queued status requests & for calling ILocalizationServiceProvider::Tick()
+	 */
+	virtual void Tick() = 0;
+
+	///**
+	// * Queues translations to have their status updated in the background.
+	// * @param	InTranslationIds	The transaltions to queue.
+	// */
+	//virtual void QueueStatusUpdate(const TArray<FLocalizationServiceTranslationIdentifier>& InTranslationIds) = 0;
+
+	/**
+	 * Check whether localization service is enabled.
+	 */
+	virtual bool IsEnabled() const = 0;
+
+	/**
+	 * Get the localization service provider that is currently in use.
+	 */
+	virtual ILocalizationServiceProvider& GetProvider() const = 0;
+
+	/**
+	 * Set the current localization service provider to the one specified here by name.
+	 * This will assert if the provider does not exist.
+	 * @param	InName	The name of the provider
+	 */
+	virtual void SetProvider( const FName& InName ) = 0;
+
+	/**
+	 * Get whether we should use global or per-project settings
+	 * @return true if we should use global settings
+	 */
+	virtual bool GetUseGlobalSettings() const = 0;
+
+	/**
+	 * Set whether we should use global or per-project settings
+	 * @param bIsUseGlobalSettings	Whether we should use global settings
+	 */
+	virtual void SetUseGlobalSettings(bool bIsUseGlobalSettings) = 0;
+
+	/**
+	 * Gets a reference to the localization service module instance.
+	 *
+	 * @return A reference to the localization service module.
+	 */
+	static inline ILocalizationServiceModule& Get()
+	{
+		static FName LocalizationServiceModule("LocalizationService");
+		return FModuleManager::LoadModuleChecked<ILocalizationServiceModule>(LocalizationServiceModule);
+	}
+};
diff --git a/Engine/Source/Developer/LocalizationService/Public/ILocalizationServiceOperation.h b/Engine/Source/Developer/LocalizationService/Public/ILocalizationServiceOperation.h
new file mode 100644
--- /dev/null
+++ b/Engine/Source/Developer/LocalizationService/Public/ILocalizationServiceOperation.h
@@ -0,0 +1,30 @@
+// Copyright 1998-2015 Epic Games, Inc. All Rights Reserved.
+
+#pragma once
+
+class ILocalizationServiceOperation : public TSharedFromThis<ILocalizationServiceOperation, ESPMode::ThreadSafe>
+{
+public:
+	/**
+	 * Virtual destructor
+	 */
+	virtual ~ILocalizationServiceOperation() {}
+
+	/** Get the name of this operation, used as a unique identifier */
+	virtual FName GetName() const = 0;
+
+	/** Get the string to display when this operation is in progress */
+	virtual FText GetInProgressString() const 
+	{ 
+		return FText(); 
+	}
+
+	/** Factory method for easier operation creation */
+	template<typename Type>
+	static TSharedRef<Type, ESPMode::ThreadSafe> Create()
+	{
+		return MakeShareable( new Type() );
+	}
+};
+
+typedef TSharedRef<class ILocalizationServiceOperation, ESPMode::ThreadSafe> FLocalizationServiceOperationRef;
diff --git a/Engine/Source/Developer/LocalizationService/Public/ILocalizationServiceProvider.h b/Engine/Source/Developer/LocalizationService/Public/ILocalizationServiceProvider.h
new file mode 100644
--- /dev/null
+++ b/Engine/Source/Developer/LocalizationService/Public/ILocalizationServiceProvider.h
@@ -0,0 +1,209 @@
+// Copyright 1998-2015 Epic Games, Inc. All Rights Reserved.
+
+#pragma once
+
+#include "LocalizationServiceOperations.h"
+//#include "LocalizationServiceHelpers.h"
+#include "Runtime/Core/Public/Features/IModularFeature.h"
+#include "ILocalizationServiceState.h"
+class IDetailCategoryBuilder;
+class ULocalizationTarget;
+
+#define LOCALIZATION_SERVICES_WITH_SLATE			!PLATFORM_LINUX
+
+/**
+ * Hint for how to execute the operation. Note that asynchronous operations require
+ * Tick() to be called to manage completed operations.
+ */
+namespace ELocalizationServiceOperationConcurrency
+{
+	enum Type
+	{
+		/** Force the operation to be issued on the same thread, blocking until complete. */
+		Synchronous,
+		/** Run the command on another thread, returning immediately. */
+		Asynchronous
+	};
+}
+
+/**
+* Hint to provider when updating state
+*/
+namespace ELocalizationServiceCacheUsage
+{
+	enum Type
+	{
+		/** Force a synchronous update of the state of the file. */
+		ForceUpdate,
+		/** Use the cached state if possible */
+		Use,
+	};
+}
+
+/**
+ * Results of a command execution
+ */
+namespace ELocalizationServiceOperationCommandResult
+{
+	enum Type
+	{
+		/** Command failed to execute correctly or was not supported by the provider. */
+		Failed,
+		/** Command executed successfully */
+		Succeeded,
+		/** Command was canceled before completion */
+		Cancelled,
+	};
+}
+
+/** Delegate used by providers for when operations finish */
+DECLARE_DELEGATE_TwoParams(FLocalizationServiceOperationComplete, const FLocalizationServiceOperationRef&, ELocalizationServiceOperationCommandResult::Type);
+
+/** Delegate used by providers to create localization service operations */
+typedef TSharedRef<class ILocalizationServiceOperation, ESPMode::ThreadSafe> FLocalizationServiceOperationSharedRef;
+DECLARE_DELEGATE_RetVal(FLocalizationServiceOperationSharedRef, FGetLocalizationServiceOperation);
+
+/**
+ * Interface to talking with localization service providers.
+ */
+class ILocalizationServiceProvider : public IModularFeature
+{
+public:
+	/**
+	 * Virtual destructor
+	 */
+	virtual ~ILocalizationServiceProvider() {}
+
+	/** 
+	 * Initialize localization service provider.
+	 * @param	bForceConnection	If set, this flag forces the provider to attempt a connection to its server.
+	 */
+	virtual void Init(bool bForceConnection = true) = 0;
+
+	/** 
+	 * Shut down localization service provider.
+	 */
+	virtual void Close() = 0;
+
+	/** Get the localization service provider name */
+	virtual const FName& GetName() const = 0;
+
+	/** Get the localization service provider display name */
+	virtual const FText GetDisplayName() const = 0;
+
+	/** Get the localization service status as plain, human-readable text */
+	virtual FText GetStatusText() const = 0;
+
+	/** Quick check if localization service is enabled */
+	virtual bool IsEnabled() const = 0;
+
+	/**
+	 * Quick check if localization service is available for use (server-based providers can use this
+	 * to return whether the server is available or not)
+	 *
+	 * @return	true if localization service is available, false if it is not
+	 */
+	virtual bool IsAvailable() const = 0;
+
+	/**
+	 * Login to the localization service server (if any).
+	 * This is just a wrapper around Execute().
+	 * @param	InPassword						The password (if required) to use when logging in.
+	 * @param	InConcurrency					How to execute the operation, blocking or asynchronously on another thread.
+	 * @param	InOperationCompleteDelegate		Delegate to call when the operation is completed. This is called back internal to this call when executed on the main thread, or from Tick() when queued for asynchronous execution.
+	 * @return the result of the operation.
+	 */
+	virtual ELocalizationServiceOperationCommandResult::Type Login(const FString& InPassword = FString(), ELocalizationServiceOperationConcurrency::Type InConcurrency = ELocalizationServiceOperationConcurrency::Synchronous, const FLocalizationServiceOperationComplete& InOperationCompleteDelegate = FLocalizationServiceOperationComplete())
+	{
+		TSharedRef<FConnectToProvider, ESPMode::ThreadSafe> ConnectOperation = ILocalizationServiceOperation::Create<FConnectToProvider>();
+		ConnectOperation->SetPassword(InPassword);
+		return Execute(ConnectOperation, FLocalizationServiceTranslationIdentifier(), InConcurrency, InOperationCompleteDelegate);
+	}
+
+	/**
+	* Get the state of each of the passed-in files. State may be cached for faster queries. Note states can be NULL!
+	* @param	InTranslationIds	The translations to retrieve state for.
+	* @param	OutState			The states of the files. This will be empty if the operation fails. Note states can be NULL!
+	* @param	InStateCacheUsage	Whether to use the state cache or to force a (synchronous) state retrieval.
+	* @return the result of the operation.
+	*/
+	virtual ELocalizationServiceOperationCommandResult::Type GetState(const TArray<FLocalizationServiceTranslationIdentifier>& InTranslationIds, TArray< TSharedRef<ILocalizationServiceState, ESPMode::ThreadSafe> >& OutState, ELocalizationServiceCacheUsage::Type InStateCacheUsage) = 0;
+
+	/**
+	* Helper overload for state retrieval, see GetState().
+	*/
+	virtual TSharedPtr<ILocalizationServiceState, ESPMode::ThreadSafe> GetState(const FLocalizationServiceTranslationIdentifier& InTranslationId, ELocalizationServiceCacheUsage::Type InStateCacheUsage)
+	{
+		TArray<FLocalizationServiceTranslationIdentifier> TranslationIds;
+		TArray< TSharedRef<ILocalizationServiceState, ESPMode::ThreadSafe> > States;
+		TranslationIds.Add(InTranslationId);
+		if (GetState(TranslationIds, States, InStateCacheUsage) == ELocalizationServiceOperationCommandResult::Succeeded)
+		{
+			TSharedRef<ILocalizationServiceState, ESPMode::ThreadSafe> State = States[0];
+			return State;
+		}
+		return NULL;
+	}
+
+	/**
+	 * Attempt to execute an operation on the passed-in files (if any are required).
+	 * @param	InTranslationIds				The translations in question
+	 * @param	InOperation						The operation to perform.
+	 * @param	InConcurrency					How to execute the operation, blocking or asynchronously on another thread.
+	 * @param	InOperationCompleteDelegate		Delegate to call when the operation is completed. This is called back internal to this call when executed on the main thread, or from Tick() when queued for asynchronous execution.
+	 * @return the result of the operation.
+	 */
+	virtual ELocalizationServiceOperationCommandResult::Type Execute(const TSharedRef<ILocalizationServiceOperation, ESPMode::ThreadSafe>& InOperation, const TArray<FLocalizationServiceTranslationIdentifier>& InTranslationIds, ELocalizationServiceOperationConcurrency::Type InConcurrency = ELocalizationServiceOperationConcurrency::Synchronous, const FLocalizationServiceOperationComplete& InOperationCompleteDelegate = FLocalizationServiceOperationComplete()) = 0;
+
+	/**
+	 * Helper overload for operation execution, see Execute().
+	 */
+	virtual ELocalizationServiceOperationCommandResult::Type Execute(const TSharedRef<ILocalizationServiceOperation, ESPMode::ThreadSafe>& InOperation, const ELocalizationServiceOperationConcurrency::Type InConcurrency = ELocalizationServiceOperationConcurrency::Synchronous, const FLocalizationServiceOperationComplete& InOperationCompleteDelegate = FLocalizationServiceOperationComplete())
+	{
+		return Execute(InOperation, TArray<FLocalizationServiceTranslationIdentifier>(), InConcurrency, InOperationCompleteDelegate);
+	}
+
+	/**
+	 * Helper overload for operation execution, see Execute().
+	 */
+	virtual ELocalizationServiceOperationCommandResult::Type Execute(const TSharedRef<ILocalizationServiceOperation, ESPMode::ThreadSafe>& InOperation, const FLocalizationServiceTranslationIdentifier& InTranslationId, const ELocalizationServiceOperationConcurrency::Type InConcurrency = ELocalizationServiceOperationConcurrency::Synchronous, const FLocalizationServiceOperationComplete& InOperationCompleteDelegate = FLocalizationServiceOperationComplete())
+	{
+		TArray<FLocalizationServiceTranslationIdentifier> TranslationIds;
+		TranslationIds.Add(InTranslationId);
+		return Execute(InOperation, TranslationIds, InConcurrency, InOperationCompleteDelegate);
+	}
+
+	/**
+	 * Check to see if we can cancel an operation.
+	 * @param	InOperation		The operation to check.
+	 * @return true if the operation was cancelled.
+	 */
+	virtual bool CanCancelOperation( const TSharedRef<ILocalizationServiceOperation, ESPMode::ThreadSafe>& InOperation ) const = 0;
+
+	/**
+	 * Attempt to cancel an operation in progress.
+	 * @param	InOperation		The operation to attempt to cancel.
+	 */
+	virtual void CancelOperation( const TSharedRef<ILocalizationServiceOperation, ESPMode::ThreadSafe>& InOperation ) = 0;
+
+	/**
+	 * Called every update.
+	 */
+	virtual void Tick() = 0;
+
+#if LOCALIZATION_SERVICES_WITH_SLATE
+	/**
+	* Create a settings widget for display in the localization dashboard.
+	* @returns a widget used to edit project-specific localization service settings.
+	*/
+	virtual void CustomizeSettingsDetails(IDetailCategoryBuilder& DetailCategoryBuilder) const = 0;
+
+	/**
+	* Create a settings widget for display in the localization target editor.
+	* @returns a widget used to edit target-specific localization service settings.
+	*/
+	virtual void CustomizeTargetDetails(IDetailCategoryBuilder& DetailCategoryBuilder, ULocalizationTarget& Target) const = 0;
+#endif //LOCALIZATION_SERVICES_WITH_SLATE
+};
+
+
diff --git a/Engine/Source/Developer/LocalizationService/Public/ILocalizationServiceRevision.h b/Engine/Source/Developer/LocalizationService/Public/ILocalizationServiceRevision.h
new file mode 100644
--- /dev/null
+++ b/Engine/Source/Developer/LocalizationService/Public/ILocalizationServiceRevision.h
@@ -0,0 +1,27 @@
+// Copyright 1998-2015 Epic Games, Inc. All Rights Reserved.
+
+#pragma once
+
+/** 
+ * Abstraction of a localization revision from a localization service provider.
+ */
+class ILocalizationServiceRevision : public TSharedFromThis<ILocalizationServiceRevision, ESPMode::ThreadSafe>
+{
+public:
+	/**
+	 * Virtual destructor
+	 */
+	virtual ~ILocalizationServiceRevision() {}
+
+	/** Number of the revision */
+	virtual int32 GetRevisionNumber() const = 0;
+
+	/** String representation of the revision */
+	virtual const FString& GetRevision() const = 0;
+
+	/** User name of the submitter */
+	virtual const FString& GetUserName() const = 0;
+
+	/** Date of the revision */
+	virtual const FDateTime& GetDate() const = 0;
+};
\ No newline at end of file
diff --git a/Engine/Source/Developer/LocalizationService/Public/ILocalizationServiceState.h b/Engine/Source/Developer/LocalizationService/Public/ILocalizationServiceState.h
new file mode 100644
--- /dev/null
+++ b/Engine/Source/Developer/LocalizationService/Public/ILocalizationServiceState.h
@@ -0,0 +1,121 @@
+// Copyright 1998-2015 Epic Games, Inc. All Rights Reserved.
+
+#pragma once
+
+typedef TSharedRef<class ILocalizationServiceState, ESPMode::ThreadSafe> FLocalizationServiceStateRef;
+typedef TSharedPtr<class ILocalizationServiceState, ESPMode::ThreadSafe> FLocalizationServiceStatePtr;
+
+struct FLocalizationServiceTranslationIdentifier //: TSharedFromThis<FLocalizationServiceTranslationIdentifier, ESPMode::ThreadSafe>
+{
+	FLocalizationServiceTranslationIdentifier()
+		: Culture(nullptr), Namespace(""), Source("")
+	{}
+
+	FLocalizationServiceTranslationIdentifier(FCulturePtr InCulturePtr, const FString& InNamespace, const FString& InSource)
+		: Culture(InCulturePtr)
+		, Namespace(InNamespace)
+		, Source(InSource)
+	{
+
+	}
+
+	FCulturePtr Culture;
+	FString Namespace;
+	FString Source;
+};
+
+/**
+ * An abstraction of the state of a file under localization service
+ */
+class ILocalizationServiceState : public TSharedFromThis<ILocalizationServiceState, ESPMode::ThreadSafe>
+{
+public:
+	enum { INVALID_REVISION = -1 };
+
+	/**
+	 * Virtual destructor
+	 */
+	virtual ~ILocalizationServiceState() {}
+
+	/** 
+	 * Get the size of the history. 
+	 * If an FUpdateStatus operation has been called with the ShouldUpdateHistory() set, there 
+	 * should be history present if the file has been committed to localization service.
+	 * @returns the number of items in the history
+	 */
+	virtual int32 GetHistorySize() const = 0;
+
+	/**
+	* Get the text in question
+	*/
+	virtual const FString& GetSourceString() const = 0;
+
+	/**
+	* Get the text in question
+	*/
+	virtual const FString& GetTranslationString() const = 0;
+
+	/**
+	* Get the culture for which the localization of the text is for
+	*/
+	virtual const FCulturePtr GetCulture() const = 0;
+
+	/**
+	 * Get an item from the history
+	 * @param	HistoryIndex	the index of the history item
+	 * @returns a history item or NULL if none exist
+	 */
+	virtual TSharedPtr<class ILocalizationServiceRevision, ESPMode::ThreadSafe> GetHistoryItem( int32 HistoryIndex ) const = 0;
+
+	/**
+	* Get the name of the icon graphic we should use to display the state in a UI.
+	* @returns the name of the icon to display
+	*/
+	virtual FName GetIconName() const = 0;
+
+	/**
+	* Get the name of the small icon graphic we should use to display the state in a UI.
+	* @returns the name of the icon to display
+	*/
+	virtual FName GetSmallIconName() const = 0;
+
+	/**
+	 * Get a text representation of the state
+	 * @returns	the text to display for this state
+	 */
+	virtual FText GetDisplayName() const = 0;
+
+	/**
+	 * Get a tooltip to describe this state
+	 * @returns	the text to display for this states tooltip
+	 */
+	virtual FText GetDisplayTooltip() const = 0;
+
+	///**
+	// * Get a unique identifier for the translation this state represents
+	// * @returns	the unique identifier
+	// */
+	virtual const FLocalizationServiceTranslationIdentifier& GetTranslationIdentifier() const = 0;
+
+	/**
+	 * Get the timestamp of the last update that was made to this state.
+	 * @returns	the timestamp of the last update
+	 */
+	virtual const FDateTime& GetTimeStamp() const = 0;
+
+	/** Get whether this translation is up-to-date with the version in localization service */
+	virtual bool IsCurrent() const = 0;
+
+	/** Get whether this text is known to the localization service */
+	virtual bool IsKnownToLocalizationService() const = 0;
+
+	/** Get whether localization service allows this translation to be edited */
+	virtual bool CanEdit() const = 0;
+
+	/** Get whether we know anything about this files localization service state */
+	virtual bool IsUnknown() const = 0;
+
+	/** Get whether this file is modified compared to the version we have from localization service */
+	virtual bool IsModified() const = 0;
+};
+
diff --git a/Engine/Source/Developer/LocalizationService/Public/LocalizationServiceHelpers.h b/Engine/Source/Developer/LocalizationService/Public/LocalizationServiceHelpers.h
new file mode 100644
--- /dev/null
+++ b/Engine/Source/Developer/LocalizationService/Public/LocalizationServiceHelpers.h
@@ -0,0 +1,42 @@
+// Copyright 1998-2015 Epic Games, Inc. All Rights Reserved.
+
+#pragma once
+
+class ILocalizationServiceProvider;
+
+namespace LocalizationServiceHelpers
+{
+	/**
+	 * Helper function to get the ini filename for storing localization service settings
+	 * @return the filename
+	 */
+	LOCALIZATIONSERVICE_API extern const FString& GetSettingsIni();
+
+	/**
+	 * Helper function to get the ini filename for storing global localization service settings
+	 * @return the filename
+	 */
+	LOCALIZATIONSERVICE_API extern const FString& GetGlobalSettingsIni();
+
+	/**
+	 * Helper function to commit a translation
+	 */
+	//LOCALIZATIONSERVICE_API extern bool CommitTranslation(const FString& Culture, const FString& Namespace, const FString& Source, const FString& Translation);
+}
+
+/**
+* Helper class that ensures FLocalizationService is properly initialized and shutdown by calling Init/Close in
+* its constructor/destructor respectively.
+*/
+class LOCALIZATIONSERVICE_API FScopedLocalizationService
+{
+public:
+	/** Constructor; Initializes localization service Provider */
+	FScopedLocalizationService();
+
+	/** Destructor; Closes localization service Provider */
+	~FScopedLocalizationService();
+
+	/** Get the provider we are using */
+	ILocalizationServiceProvider& GetProvider();
+};
\ No newline at end of file
diff --git a/Engine/Source/Developer/LocalizationService/Public/LocalizationServiceOperations.h b/Engine/Source/Developer/LocalizationService/Public/LocalizationServiceOperations.h
new file mode 100644
--- /dev/null
+++ b/Engine/Source/Developer/LocalizationService/Public/LocalizationServiceOperations.h
@@ -0,0 +1,129 @@
+// Copyright 1998-2015 Epic Games, Inc. All Rights Reserved.
+
+#pragma once
+
+#include "ILocalizationServiceOperation.h"
+
+#define LOCTEXT_NAMESPACE "LocalizationService"
+
+/**
+ * Operation used to connect (or test a connection) to localization service
+ */
+class FConnectToProvider : public ILocalizationServiceOperation
+{
+public:
+	// ILocalizationServiceOperation interface
+	virtual FName GetName() const override 
+	{ 
+		return "Connect"; 
+	}
+
+	virtual FText GetInProgressString() const override
+	{ 
+		return LOCTEXT("LocalizationService_Connecting", "Connecting to localization service..."); 
+	}
+
+	const FString& GetPassword() const
+	{
+		return Password;
+	}
+
+	void SetPassword(const FString& InPassword)
+	{
+		Password = InPassword;
+	}
+
+	const FText& GetErrorText() const
+	{
+		return OutErrorText;
+	}
+
+	void SetErrorText(const FText& InErrorText)
+	{
+		OutErrorText = InErrorText;
+	}
+
+protected:
+	/** Password we use for this operation */
+	FString Password;
+
+	/** Error text for easy diagnosis */
+	FText OutErrorText;
+};
+
+///**
+// * Operation used to commit localization(s) to localization service
+// */
+//class FCommitTranslations : public ILocalizationServiceOperation
+//{
+//public:
+//	// ILocalizationServiceOperation interface
+//	virtual FName GetName() const override
+//	{
+//		return "CommitTranslations";
+//	}
+//
+//	virtual FText GetInProgressString() const override
+//	{ 
+//		return LOCTEXT("LocalizationService_FCommit", "Committing translation(s) to localization Service..."); 
+//	}
+//
+//
+//	void SetSuccessMessage( const FText& InSuccessMessage )
+//	{
+//		SuccessMessage = InSuccessMessage;
+//	}
+//
+//	const FText& GetSuccessMessage() const
+//	{
+//		return SuccessMessage;
+//	}
+//
+//protected:
+//
+//	/** A short message listing changelist/revision we submitted, if successful */
+//	FText SuccessMessage;
+//
+//	/** The translations we are committing */
+//	TArray<TPair<FLocalizationServiceTranslationIdentifier, FString>> TranslationsToCommit;
+//};
+
+///**
+//* Operation used to update the localization service status of files
+//*/
+//class FLocalizationServiceUpdateStatus : public ILocalizationServiceOperation
+//{
+//public:
+//	FLocalizationServiceUpdateStatus()
+//		: bUpdateHistory(false)
+//	{
+//	}
+//
+//	// ILocalizationServiceOperation interface
+//	virtual FName GetName() const override
+//	{
+//		return "UpdateStatus";
+//	}
+//
+//	virtual FText GetInProgressString() const override
+//	{
+//		return LOCTEXT("LocalizationService_Update", "Updating translation(s) localization service status...");
+//	}
+//
+//	void SetUpdateHistory(bool bInUpdateHistory)
+//	{
+//		bUpdateHistory = bInUpdateHistory;
+//	}
+//
+//	bool ShouldUpdateHistory() const
+//	{
+//		return bUpdateHistory;
+//	}
+//	
+//protected:
+//	/** Whether to update history */
+//	bool bUpdateHistory;
+//};
+
+
+#undef LOCTEXT_NAMESPACE
\ No newline at end of file
diff --git a/Engine/Source/Developer/LocalizationService/Public/ScopedLocalizationServiceProgress.h b/Engine/Source/Developer/LocalizationService/Public/ScopedLocalizationServiceProgress.h
new file mode 100644
--- /dev/null
+++ b/Engine/Source/Developer/LocalizationService/Public/ScopedLocalizationServiceProgress.h
@@ -0,0 +1,27 @@
+// Copyright 1998-2015 Epic Games, Inc. All Rights Reserved.
+
+#pragma once
+
+class FScopedLocalizationServiceProgress
+{
+public:
+	/** 
+	 * Constructor.
+	 * Progress dialog is shown immediately (if possible) and is closed when this
+	 * object goes out of scope.
+	 * @param	InText			The text to display in the progress dialog. If this is empty the dialog will not be shown.
+	 * @param	InOnCancelled	The delegate to call when the cancel button is clicked.
+	 */
+	LOCALIZATIONSERVICE_API FScopedLocalizationServiceProgress(const FText& InText, const FSimpleDelegate& InOnCancelled = FSimpleDelegate());
+
+	/** Destructor - closes the dialog */
+	LOCALIZATIONSERVICE_API ~FScopedLocalizationServiceProgress();
+
+	/** Makes sure the dialog is ticked */
+	LOCALIZATIONSERVICE_API void Tick();
+
+private:
+
+	/** The window containing this progress window */
+	TWeakPtr<class SWindow> WindowPtr;
+};
\ No newline at end of file
