Build fix for Main (Win64 and PS4) for CL 34950439 where some includes were moved

#trivial

diff --git a/Engine/Source/Developer/LocalizationService/Private/DefaultLocalizationServiceProvider.cpp b/Engine/Source/Developer/LocalizationService/Private/DefaultLocalizationServiceProvider.cpp
--- a/Engine/Source/Developer/LocalizationService/Private/DefaultLocalizationServiceProvider.cpp
+++ b/Engine/Source/Developer/LocalizationService/Private/DefaultLocalizationServiceProvider.cpp
@@ -2,10 +2,7 @@
 
 #include "DefaultLocalizationServiceProvider.h"
 #include "UObject/WeakObjectPtr.h"
-
-#if WITH_EDITOR
 #include "Logging/MessageLog.h"
-#endif
 
 #define LOCTEXT_NAMESPACE "DefaultLocalizationServiceProvider"
 
