[LocalizationService] Fix error when building RuntimeAndProgram modules on Linux.

Build.cs vs #define conditions for Slate inclusion were inconsistent.

[RN] Fixed ability to build RuntimeAndProgram modules on Linux that use localization.

diff --git a/Engine/Source/Developer/LocalizationService/LocalizationService.Build.cs b/Engine/Source/Developer/LocalizationService/LocalizationService.Build.cs
--- a/Engine/Source/Developer/LocalizationService/LocalizationService.Build.cs
+++ b/Engine/Source/Developer/LocalizationService/LocalizationService.Build.cs
@@ -15,7 +15,7 @@ public class LocalizationService : ModuleRules
 			}
 		);
 
-        if (!Target.IsInPlatformGroup(UnrealPlatformGroup.Linux))
+        if (!(Target.IsInPlatformGroup(UnrealPlatformGroup.Linux) && Target.Type == TargetType.Program))
         {
             PrivateDependencyModuleNames.AddRange(
                 new string[] {
